Settings: code cleanup

diff --git a/Engine/Source/Runtime/EngineSettings/Private/EngineSettingsModule.cpp b/Engine/Source/Runtime/EngineSettings/Private/EngineSettingsModule.cpp
--- a/Engine/Source/Runtime/EngineSettings/Private/EngineSettingsModule.cpp
+++ b/Engine/Source/Runtime/EngineSettings/Private/EngineSettingsModule.cpp
@@ -15,6 +15,8 @@ class FEngineSettingsModule
 {
 public:
 
+	// Begin IModuleInterface interface
+
 	virtual void StartupModule( ) OVERRIDE { }
 
 	virtual void ShutdownModule( ) OVERRIDE { }
@@ -23,11 +25,22 @@ public:
 	{
 		return true;
 	}
+
+	// End IModuleInterface interface
 };
 
+
 /* Class constructors
  *****************************************************************************/
 
+UGameMapsSettings::UGameMapsSettings( const class FPostConstructInitializeProperties& PCIP )
+	: Super(PCIP)
+	, bUseSplitscreen(true)
+	, TwoPlayerSplitscreenLayout(ETwoPlayerSplitScreenType::Horizontal)
+	, ThreePlayerSplitscreenLayout(EThreePlayerSplitScreenType::FavorTop)
+{ }
+
+
 UGameNetworkManagerSettings::UGameNetworkManagerSettings( const class FPostConstructInitializeProperties& PCIP )
 	: Super(PCIP)
 { }
@@ -52,17 +65,10 @@ UHudSettings::UHudSettings( const class FPostConstructInitializeProperties& PCIP
 	: Super(PCIP)
 { }
 
+
 /* Static functions
  *****************************************************************************/
 
-UGameMapsSettings::UGameMapsSettings(const class FPostConstructInitializeProperties& PCIP)
-	: Super(PCIP)
-	, bUseSplitscreen(true)
-	, TwoPlayerSplitscreenLayout(ETwoPlayerSplitScreenType::Horizontal)
-	, ThreePlayerSplitscreenLayout(EThreePlayerSplitScreenType::FavorTop)
-{ 
-}
-
 const FString& UGameMapsSettings::GetGameDefaultMap( )
 {
 	return IsRunningDedicatedServer()
