Default to using the GlobalDefaultGameMode if GlobalDefaultServerGameMode is not specified.
Cleared out places setting GlobalDefaultServerGameMode that need not do so.

diff --git a/Engine/Config/BaseEngine.ini b/Engine/Config/BaseEngine.ini
--- a/Engine/Config/BaseEngine.ini
+++ b/Engine/Config/BaseEngine.ini
@@ -5,7 +5,6 @@ EditorStartupMap=/Engine/Maps/Templates/Template_Default
 GameDefaultMap=/Engine/Maps/Entry
 ServerDefaultMap=/Engine/Maps/Entry
 GlobalDefaultGameMode="/Script/Engine.GameMode"
-GlobalDefaultServerGameMode="/Script/Engine.GameMode"
 LocalMapOptions=
 
 [URL]
diff --git a/Engine/Source/Editor/GameProjectGeneration/Private/GameProjectUtils.cpp b/Engine/Source/Editor/GameProjectGeneration/Private/GameProjectUtils.cpp
--- a/Engine/Source/Editor/GameProjectGeneration/Private/GameProjectUtils.cpp
+++ b/Engine/Source/Editor/GameProjectGeneration/Private/GameProjectUtils.cpp
@@ -1465,7 +1465,6 @@ bool GameProjectUtils::GenerateConfigFiles(const FString& NewProjectPath, const
 				if (bShouldGenerateCode)
 				{
 					FileContents += FString::Printf(TEXT("GlobalDefaultGameMode=\"/Script/%s.%sGameMode\"") LINE_TERMINATOR, *NewProjectName, *NewProjectName);
-					FileContents += FString::Printf(TEXT("GlobalDefaultServerGameMode=\"/Script/%s.%sGameMode\"") LINE_TERMINATOR, *NewProjectName, *NewProjectName);
 				}
 			}
 		}
diff --git a/Engine/Source/Runtime/EngineSettings/Classes/GameMapsSettings.h b/Engine/Source/Runtime/EngineSettings/Classes/GameMapsSettings.h
--- a/Engine/Source/Runtime/EngineSettings/Classes/GameMapsSettings.h
+++ b/Engine/Source/Runtime/EngineSettings/Classes/GameMapsSettings.h
@@ -110,7 +110,8 @@ private:
 	FStringClassReference GlobalDefaultGameMode;
 
 	/**
-	 * GameMode to use if not specified in any other way. (e.g. per-map DefaultGameMode or on the URL) (DEDICATED SERVERS).
+	 * GameMode to use if not specified in any other way. (e.g. per-map DefaultGameMode or on the URL) (DEDICATED SERVERS)
+	 * If not set, the GlobalDefaultGameMode value will be used.
 	 */
 	UPROPERTY(config, EditAnywhere, Category=DefaultModes, meta=(MetaClass="GameMode"), AdvancedDisplay)
 	FStringClassReference GlobalDefaultServerGameMode;
diff --git a/Engine/Source/Runtime/EngineSettings/Private/EngineSettingsModule.cpp b/Engine/Source/Runtime/EngineSettings/Private/EngineSettingsModule.cpp
--- a/Engine/Source/Runtime/EngineSettings/Private/EngineSettingsModule.cpp
+++ b/Engine/Source/Runtime/EngineSettings/Private/EngineSettingsModule.cpp
@@ -79,9 +79,11 @@ const FString& UGameMapsSettings::GetGameDefaultMap( )
 
 const FString& UGameMapsSettings::GetGlobalDefaultGameMode( )
 {
-	return IsRunningDedicatedServer()
-		? GetDefault<UGameMapsSettings>()->GlobalDefaultServerGameMode.ClassName 
-		: GetDefault<UGameMapsSettings>()->GlobalDefaultGameMode.ClassName;
+	UGameMapsSettings* GameMapsSettings = Cast<UGameMapsSettings>(UGameMapsSettings::StaticClass()->GetDefaultObject());
+
+	return IsRunningDedicatedServer() && GameMapsSettings->GlobalDefaultServerGameMode.IsValid()
+		? GameMapsSettings->GlobalDefaultServerGameMode.ClassName
+		: GameMapsSettings->GlobalDefaultGameMode.ClassName;
 }
 
 
