Comment split screen and game instance project setting properties

diff --git a/Engine/Source/Runtime/EngineSettings/Classes/GameMapsSettings.h b/Engine/Source/Runtime/EngineSettings/Classes/GameMapsSettings.h
--- a/Engine/Source/Runtime/EngineSettings/Classes/GameMapsSettings.h
+++ b/Engine/Source/Runtime/EngineSettings/Classes/GameMapsSettings.h
@@ -5,6 +5,7 @@
 #include "GameMapsSettings.generated.h"
 
 
+// Ways the screen can be split with two players
 UENUM()
 namespace ETwoPlayerSplitScreenType
 {
@@ -16,6 +17,7 @@ namespace ETwoPlayerSplitScreenType
 }
 
 
+// Ways the screen can be split with three players
 UENUM()
 namespace EThreePlayerSplitScreenType
 {
@@ -77,15 +79,19 @@ public:
 	UPROPERTY(config, EditAnywhere, Category=DefaultMaps, AdvancedDisplay)
 	FString TransitionMap;
 
+	/** Whether the screen should be split or not when multiple local players are present */
 	UPROPERTY(config, EditAnywhere, Category=LocalMultiplayer)
 	bool bUseSplitscreen;
 
+	/** The viewport layout to use if the screen should be split and there are two local players */
 	UPROPERTY(config, EditAnywhere, Category=LocalMultiplayer, meta=(editcondition="bUseSplitScreen"))
 	TEnumAsByte<ETwoPlayerSplitScreenType::Type> TwoPlayerSplitscreenLayout;
 
+	/** The viewport layout to use if the screen should be split and there are three local players */
 	UPROPERTY(config, EditAnywhere, Category=LocalMultiplayer, meta=(editcondition="bUseSplitScreen"))
 	TEnumAsByte<EThreePlayerSplitScreenType::Type> ThreePlayerSplitscreenLayout;
 
+	/** The class to use when instantiating the transient GameInstance class */
 	UPROPERTY(config, noclear, EditAnywhere, Category=GameInstance, meta=(MetaClass="GameInstance"))
 	FStringClassReference GameInstanceClass;
 
