Expose UGameMapSettings::SkipAssigningGamepadToPlayer1 to blueprints

diff --git a/Engine/Source/Runtime/EngineSettings/Classes/GameMapsSettings.h b/Engine/Source/Runtime/EngineSettings/Classes/GameMapsSettings.h
--- a/Engine/Source/Runtime/EngineSettings/Classes/GameMapsSettings.h
+++ b/Engine/Source/Runtime/EngineSettings/Classes/GameMapsSettings.h
@@ -188,4 +188,22 @@ private:
 	/** List of GameModes to load when game= is specified in the URL (e.g. "DM" could be an alias for "MyProject.MyGameModeMP_DM") */
 	UPROPERTY(config, EditAnywhere, Category = DefaultModes, AdvancedDisplay)
 	TArray<FGameModeName> GameModeClassAliases;
+
+public:
+
+	/** Returns the game local maps settings */
+	UFUNCTION(BlueprintPure, Category = Settings, meta=(DisplayName="Get Game Maps and Modes Settings"))
+	static UGameMapsSettings* GetGameMapsSettings();
+
+	/**
+	 * Modify "Skip Assigning Gamepad to Player 1" GameMapsSettings option
+	 * @param bSkipFirstPlayer		If set connected game pads will only be assigned to the second and subsequent players
+	 * @note This value is saved to local config when changed.
+	*/
+	UFUNCTION(BlueprintCallable, Category = "Utilities")
+	void SetSkipAssigningGamepadToPlayer1(bool bSkipFirstPlayer = true);
+
+	UFUNCTION(BlueprintPure, Category = "Utilities")
+	bool GetSkipAssigningGamepadToPlayer1() const;
+
 };
diff --git a/Engine/Source/Runtime/EngineSettings/Private/EngineSettingsModule.cpp b/Engine/Source/Runtime/EngineSettings/Private/EngineSettingsModule.cpp
--- a/Engine/Source/Runtime/EngineSettings/Private/EngineSettingsModule.cpp
+++ b/Engine/Source/Runtime/EngineSettings/Private/EngineSettingsModule.cpp
@@ -213,6 +213,21 @@ void UGameMapsSettings::SetGlobalDefaultGameMode( const FString& NewGameMode )
 	GameMapsSettings->GlobalDefaultGameMode = NewGameMode;
 }
 
+UGameMapsSettings* UGameMapsSettings::GetGameMapsSettings()
+{
+	return GetMutableDefault<UGameMapsSettings>();
+}
+
+void UGameMapsSettings::SetSkipAssigningGamepadToPlayer1(bool bSkipFirstPlayer)
+{
+	bOffsetPlayerGamepadIds = bSkipFirstPlayer;
+}
+
+bool UGameMapsSettings::GetSkipAssigningGamepadToPlayer1() const
+{
+	return bOffsetPlayerGamepadIds;
+}
+
 // Backwards compat for map strings
 void FixMapAssetRef(FSoftObjectPath& MapAssetReference)
 {
