Switching console help display to the implementation already added to UE main in CL# 8216634 (first line only, behind an option), but enabling that option by default
#rnx

diff --git a/Engine/Source/Runtime/EngineSettings/Classes/ConsoleSettings.h b/Engine/Source/Runtime/EngineSettings/Classes/ConsoleSettings.h
--- a/Engine/Source/Runtime/EngineSettings/Classes/ConsoleSettings.h
+++ b/Engine/Source/Runtime/EngineSettings/Classes/ConsoleSettings.h
@@ -88,7 +88,7 @@ class ENGINESETTINGS_API UConsoleSettings
 
 	/** Display the first line of any available help text in the auto-complete window, if a description isn't available */
 	UPROPERTY(config, EditAnywhere, Category = AutoComplete)
-	bool bDisplayHelpInAutoComplete;
+	bool bDisplayHelpInAutoComplete = true;
 
 	/** The color used for text input. */
 	UPROPERTY(config, EditAnywhere, Category=Colors)
