Added option for displaying first line of help text in console auto-complete window. New option is disabled by default.

diff --git a/Engine/Source/Runtime/Engine/Private/UserInterface/Console.cpp b/Engine/Source/Runtime/Engine/Private/UserInterface/Console.cpp
--- a/Engine/Source/Runtime/Engine/Private/UserInterface/Console.cpp
+++ b/Engine/Source/Runtime/Engine/Private/UserInterface/Console.cpp
@@ -96,6 +96,16 @@ public:
 		FAutoCompleteCommand& Cmd = Sink[NewIdx];
 		Cmd.Command = Name;
 
+		if (ConsoleSettings->bDisplayHelpInAutoComplete)
+		{
+			TArray<FString> Lines;
+			FString(CVar->GetHelp()).ParseIntoArrayLines(Lines, true);
+			if (Lines.Num())
+			{
+				Cmd.Desc = Lines[0];
+			}
+		}
+
 		IConsoleVariable* CVariable = CVar->AsVariable();
 		if (CVariable)
 		{
diff --git a/Engine/Source/Runtime/EngineSettings/Classes/ConsoleSettings.h b/Engine/Source/Runtime/EngineSettings/Classes/ConsoleSettings.h
--- a/Engine/Source/Runtime/EngineSettings/Classes/ConsoleSettings.h
+++ b/Engine/Source/Runtime/EngineSettings/Classes/ConsoleSettings.h
@@ -86,6 +86,10 @@ class ENGINESETTINGS_API UConsoleSettings
 	UPROPERTY(config, EditAnywhere, Category = AutoComplete)
 	bool bOrderTopToBottom;
 
+	/** Display the first line of any available help text in the auto-complete window, if a description isn't available */
+	UPROPERTY(config, EditAnywhere, Category = AutoComplete)
+	bool bDisplayHelpInAutoComplete;
+
 	/** The color used for text input. */
 	UPROPERTY(config, EditAnywhere, Category=Colors)
 	FColor InputColor;
