Fixed error compiling non-editor platforms.

diff --git a/Engine/Source/Runtime/EngineSettings/Private/EngineSettingsModule.cpp b/Engine/Source/Runtime/EngineSettings/Private/EngineSettingsModule.cpp
--- a/Engine/Source/Runtime/EngineSettings/Private/EngineSettingsModule.cpp
+++ b/Engine/Source/Runtime/EngineSettings/Private/EngineSettingsModule.cpp
@@ -230,7 +230,9 @@ void UGameMapsSettings::PostInitProperties()
 {
 	Super::PostInitProperties();
 
+#if WITH_EDITORONLY_DATA
 	FixMapAssetRef(EditorStartupMap);
+#endif
 	FixMapAssetRef(GameDefaultMap);
 	FixMapAssetRef(ServerDefaultMap);
 	FixMapAssetRef(TransitionMap);
@@ -240,11 +242,14 @@ void UGameMapsSettings::PostReloadConfig( UProperty* PropertyThatWasLoaded )
 {
 	if (PropertyThatWasLoaded)
 	{
+#if WITH_EDITORONLY_DATA
 		if (PropertyThatWasLoaded->GetFName() == GET_MEMBER_NAME_CHECKED(UGameMapsSettings, EditorStartupMap))
 		{
 			FixMapAssetRef(EditorStartupMap);
 		}
-		else if (PropertyThatWasLoaded->GetFName() == GET_MEMBER_NAME_CHECKED(UGameMapsSettings, GameDefaultMap))
+		else
+#endif
+		if (PropertyThatWasLoaded->GetFName() == GET_MEMBER_NAME_CHECKED(UGameMapsSettings, GameDefaultMap))
 		{
 			FixMapAssetRef(GameDefaultMap);
 		}
@@ -259,7 +264,9 @@ void UGameMapsSettings::PostReloadConfig( UProperty* PropertyThatWasLoaded )
 	}
 	else
 	{
+#if WITH_EDITORONLY_DATA
 		FixMapAssetRef(EditorStartupMap);
+#endif
 		FixMapAssetRef(GameDefaultMap);
 		FixMapAssetRef(ServerDefaultMap);
 		FixMapAssetRef(TransitionMap);
