PR #9224: Expose fake stereo class parameters (Contributed by [REDACTED_PERSON])

diff --git a/Engine/Source/Runtime/Engine/Private/UnrealEngine.cpp b/Engine/Source/Runtime/Engine/Private/UnrealEngine.cpp
--- a/Engine/Source/Runtime/Engine/Private/UnrealEngine.cpp
+++ b/Engine/Source/Runtime/Engine/Private/UnrealEngine.cpp
@@ -3475,7 +3475,7 @@ class FFakeStereoRenderingDevice : public IStereoRendering
 {
 public:
 	FFakeStereoRenderingDevice() 
-		: FOVInDegrees(100)
+		: FOVInDegrees(GetDefault<UGeneralProjectSettings>()->FOVForFakeStereoRenderingDevice)
 		, Width(640)
 		, Height(480)
 		, NumViews(2)
@@ -3544,7 +3544,7 @@ public:
 	virtual void CalculateStereoViewOffset(const int32 ViewIndex, FRotator& ViewRotation, const float WorldToMeters, FVector& ViewLocation) override
 	{
 		// 32mm, 1/2 average interpupillary distance
-		float EyeOffset = .0320000005f * WorldToMeters;
+		float EyeOffset = GetDefault<UGeneralProjectSettings>()->EyeOffsetForFakeStereoRenderingDevice * WorldToMeters;
 
 		if (ViewIndex == eSSE_MONOSCOPIC)
 		{
diff --git a/Engine/Source/Runtime/EngineSettings/Classes/GeneralProjectSettings.h b/Engine/Source/Runtime/EngineSettings/Classes/GeneralProjectSettings.h
--- a/Engine/Source/Runtime/EngineSettings/Classes/GeneralProjectSettings.h
+++ b/Engine/Source/Runtime/EngineSettings/Classes/GeneralProjectSettings.h
@@ -94,4 +94,12 @@ class ENGINESETTINGS_API UGeneralProjectSettings
 	/** Should a minimize button be shown for the game's window, when not using full screen */
 	UPROPERTY(config, EditAnywhere, Category = Settings)
 	bool bAllowMinimize;
+
+	/*Determines the Eye offset of the virtual stereo device created when " -emulatestereo" command line arg is detected*/
+	UPROPERTY(config, EditAnywhere, Category = Settings)
+	float EyeOffsetForFakeStereoRenderingDevice;
+
+	/*Determines the Field Of View of the virtual stereo device created when " -emulatestereo" command line arg is detected*/
+	UPROPERTY(config, EditAnywhere, Category = Settings)
+	float FOVForFakeStereoRenderingDevice;
 };
diff --git a/Engine/Source/Runtime/EngineSettings/Private/EngineSettingsModule.cpp b/Engine/Source/Runtime/EngineSettings/Private/EngineSettingsModule.cpp
--- a/Engine/Source/Runtime/EngineSettings/Private/EngineSettingsModule.cpp
+++ b/Engine/Source/Runtime/EngineSettings/Private/EngineSettingsModule.cpp
@@ -75,6 +75,8 @@ UGeneralProjectSettings::UGeneralProjectSettings( const FObjectInitializer& Obje
 	, bAllowClose(true)
 	, bAllowMaximize(true)
 	, bAllowMinimize(true)
+	, EyeOffsetForFakeStereoRenderingDevice(0.0320000005f)
+	, FOVForFakeStereoRenderingDevice(100)
 { }
 
 
