Improve the description of the MaxScrollbackSize project setting.

diff --git a/Engine/Source/Runtime/EngineSettings/Classes/ConsoleSettings.h b/Engine/Source/Runtime/EngineSettings/Classes/ConsoleSettings.h
--- a/Engine/Source/Runtime/EngineSettings/Classes/ConsoleSettings.h
+++ b/Engine/Source/Runtime/EngineSettings/Classes/ConsoleSettings.h
@@ -66,7 +66,7 @@ class ENGINESETTINGS_API UConsoleSettings
 {
 	GENERATED_UCLASS_BODY()
 
-	/**  Visible Console stuff */
+	/** The maximum number of lines the console keeps in its output history. */
 	UPROPERTY(globalconfig, EditAnywhere, Category=General)
 	int32 MaxScrollbackSize;
 
