Add support for manual auto-complete commands in the Output Log console

diff --git a/Engine/Plugins/Developer/Concert/ConcertApp/MultiUserServer/Source/MultiUserServer/MultiUserServer.build.cs b/Engine/Plugins/Developer/Concert/ConcertApp/MultiUserServer/Source/MultiUserServer/MultiUserServer.build.cs
--- a/Engine/Plugins/Developer/Concert/ConcertApp/MultiUserServer/Source/MultiUserServer/MultiUserServer.build.cs
+++ b/Engine/Plugins/Developer/Concert/ConcertApp/MultiUserServer/Source/MultiUserServer/MultiUserServer.build.cs
@@ -19,6 +19,7 @@ namespace UnrealBuildTool.Rules
 					"ConcertSharedSlate",
 					"ConcertSyncCore",
 					"ConcertTransport",
+					"EngineSettings",
 					"InputCore",
 					"Messaging",
 					"OutputLog",
diff --git a/Engine/Plugins/Developer/Concert/ConcertApp/MultiUserServer/Source/MultiUserServer/Private/ConcertConsoleCommandExecutor.cpp b/Engine/Plugins/Developer/Concert/ConcertApp/MultiUserServer/Source/MultiUserServer/Private/ConcertConsoleCommandExecutor.cpp
--- a/Engine/Plugins/Developer/Concert/ConcertApp/MultiUserServer/Source/MultiUserServer/Private/ConcertConsoleCommandExecutor.cpp
+++ b/Engine/Plugins/Developer/Concert/ConcertApp/MultiUserServer/Source/MultiUserServer/Private/ConcertConsoleCommandExecutor.cpp
@@ -2,6 +2,7 @@
 
 #include "ConcertConsoleCommandExecutor.h"
 
+#include "ConsoleSettings.h"
 #include "Framework/Commands/InputChord.h"
 
 #define LOCTEXT_NAMESPACE "UnrealMultiUserUI"
@@ -47,6 +48,7 @@ namespace UE::MultiUserServer
 		};
 
 		IConsoleManager::Get().ForEachConsoleObjectThatContains(FConsoleObjectVisitor::CreateLambda(OnConsoleVariable), Input);
+		Out.Append(GetDefault<UConsoleSettings>()->GetFilteredManualAutoCompleteCommands(Input));
 	}
 
 	void FConcertConsoleCommandExecutor::GetExecHistory(TArray<FString>& Out)
diff --git a/Engine/Source/Developer/OutputLog/OutputLog.Build.cs b/Engine/Source/Developer/OutputLog/OutputLog.Build.cs
--- a/Engine/Source/Developer/OutputLog/OutputLog.Build.cs
+++ b/Engine/Source/Developer/OutputLog/OutputLog.Build.cs
@@ -10,6 +10,7 @@ public class OutputLog : ModuleRules
 			new string[] {
 				"Core",
 				"CoreUObject", // @todo Mac: for some reason it's needed to link in debug on Mac
+				"EngineSettings",
 				"InputCore",
 				"Slate",
 				"SlateCore",
diff --git a/Engine/Source/Developer/OutputLog/Private/SOutputLog.cpp b/Engine/Source/Developer/OutputLog/Private/SOutputLog.cpp
--- a/Engine/Source/Developer/OutputLog/Private/SOutputLog.cpp
+++ b/Engine/Source/Developer/OutputLog/Private/SOutputLog.cpp
@@ -1,6 +1,7 @@
 // Copyright Epic Games, Inc. All Rights Reserved.
 
 #include "SOutputLog.h"
+#include "ConsoleSettings.h"
 #include "Framework/Text/IRun.h"
 #include "Framework/Text/TextLayout.h"
 #include "Misc/ConfigCacheIni.h"
@@ -363,6 +364,7 @@ void SConsoleInputBox::OnTextChanged(const FText& InText)
 			};
 
 			IConsoleManager::Get().ForEachConsoleObjectThatContains(FConsoleObjectVisitor::CreateLambda(OnConsoleVariable), *InputTextStr);
+			AutoCompleteList.Append(GetDefault<UConsoleSettings>()->GetFilteredManualAutoCompleteCommands(InputTextStr));
 		}
 		AutoCompleteList.Sort([InputTextStr](const FString& A, const FString& B)
 		{ 
diff --git a/Engine/Source/Editor/UnrealEd/Private/Toolkits/FConsoleCommandExecutor.cpp b/Engine/Source/Editor/UnrealEd/Private/Toolkits/FConsoleCommandExecutor.cpp
--- a/Engine/Source/Editor/UnrealEd/Private/Toolkits/FConsoleCommandExecutor.cpp
+++ b/Engine/Source/Editor/UnrealEd/Private/Toolkits/FConsoleCommandExecutor.cpp
@@ -1,6 +1,7 @@
 // Copyright Epic Games, Inc. All Rights Reserved.
 
 #include "Toolkits/FConsoleCommandExecutor.h"
+#include "ConsoleSettings.h"
 #include "EngineGlobals.h"
 #include "Editor.h"
 #include "Toolkits/GlobalEditorCommonCommands.h"
@@ -54,6 +55,7 @@ void FConsoleCommandExecutor::GetAutoCompleteSuggestions(const TCHAR* Input, TAr
 	};
 
 	IConsoleManager::Get().ForEachConsoleObjectThatContains(FConsoleObjectVisitor::CreateLambda(OnConsoleVariable), Input);
+	Out.Append(GetDefault<UConsoleSettings>()->GetFilteredManualAutoCompleteCommands(Input));
 }
 
 void FConsoleCommandExecutor::GetExecHistory(TArray<FString>& Out)
diff --git a/Engine/Source/Runtime/Engine/Private/Engine/EngineConsoleCommandExecutor.cpp b/Engine/Source/Runtime/Engine/Private/Engine/EngineConsoleCommandExecutor.cpp
--- a/Engine/Source/Runtime/Engine/Private/Engine/EngineConsoleCommandExecutor.cpp
+++ b/Engine/Source/Runtime/Engine/Private/Engine/EngineConsoleCommandExecutor.cpp
@@ -2,6 +2,7 @@
 
 #include "Engine/EngineConsoleCommandExecutor.h"
 
+#include "ConsoleSettings.h"
 #include "Engine/GameEngine.h"
 #include "Engine/LocalPlayer.h"
 #include "Features/IModularFeatures.h"
@@ -65,6 +66,7 @@ void FEngineConsoleCommandExecutor::GetAutoCompleteSuggestions(const TCHAR* Inpu
 	};
 
 	IConsoleManager::Get().ForEachConsoleObjectThatContains(FConsoleObjectVisitor::CreateLambda(OnConsoleVariable), Input);
+	Out.Append(GetDefault<UConsoleSettings>()->GetFilteredManualAutoCompleteCommands(Input));
 }
 
 void FEngineConsoleCommandExecutor::GetExecHistory(TArray<FString>& Out)
diff --git a/Engine/Source/Runtime/EngineSettings/Classes/ConsoleSettings.h b/Engine/Source/Runtime/EngineSettings/Classes/ConsoleSettings.h
--- a/Engine/Source/Runtime/EngineSettings/Classes/ConsoleSettings.h
+++ b/Engine/Source/Runtime/EngineSettings/Classes/ConsoleSettings.h
@@ -109,4 +109,7 @@ class UConsoleSettings
 	/** The autocomplete color used for command descriptions and read-only CVars. */
 	UPROPERTY(config, EditAnywhere, Category=Colors)
 	FColor AutoCompleteFadedColor;
+
+	/** Returns an array with the auto-complete command names that match all the words in the substring. */
+	ENGINESETTINGS_API TArray<FString> GetFilteredManualAutoCompleteCommands(FStringView Substring) const;
 };
diff --git a/Engine/Source/Runtime/EngineSettings/Private/ConsoleSettings.cpp b/Engine/Source/Runtime/EngineSettings/Private/ConsoleSettings.cpp
new file mode 100644
--- /dev/null
+++ b/Engine/Source/Runtime/EngineSettings/Private/ConsoleSettings.cpp
@@ -0,0 +1,35 @@
+﻿// Copyright Epic Games, Inc. All Rights Reserved.
+
+#include "ConsoleSettings.h"
+
+TArray<FString> UConsoleSettings::GetFilteredManualAutoCompleteCommands(FStringView Substring) const
+{
+	// We reserve some arbitrary size here, we could use more complex heuristics later on
+	TArray<FString> Commands;
+	Commands.Reserve(Substring.IsEmpty() ? ManualAutoCompleteList.Num() : 20);
+
+	// We convert to FString because there's no ParseIntoArrayWS just yet that's compatible with String Views
+	TArray<FString> Tokens;
+	FString(Substring).ParseIntoArrayWS(Tokens);
+
+	for (const FAutoCompleteCommand& Command : ManualAutoCompleteList)
+	{
+		bool bAllMatch = true;
+
+		for (const FString& Token : Tokens)
+		{
+			if (!Command.Command.Contains(Token))
+			{
+				bAllMatch = false;
+				break;
+			}
+		}
+
+		if (bAllMatch)
+		{
+			Commands.Add(Command.Command);
+		}
+	}
+
+	return Commands;
+}
