[Console]
- Convert FAutoCompleteCommand to an aggregate type by removing it's default constructor. Use a default member initializer for Color instead.
- This allows FAutoCompleteCommand to be aggregate initialized (with designated initializers too if the compiler supports it)

diff --git a/Engine/Source/Runtime/EngineSettings/Classes/ConsoleSettings.h b/Engine/Source/Runtime/EngineSettings/Classes/ConsoleSettings.h
--- a/Engine/Source/Runtime/EngineSettings/Classes/ConsoleSettings.h
+++ b/Engine/Source/Runtime/EngineSettings/Classes/ConsoleSettings.h
@@ -22,11 +22,7 @@ struct FAutoCompleteCommand
 	UPROPERTY(config, EditAnywhere, Category=Command)
 	FString Desc;
 
-	FColor Color;
-
-	FAutoCompleteCommand()
-		: Color(180, 180, 180)
-	{}
+	FColor Color = FColor(180, 180, 180);
 
 	bool operator<(const FAutoCompleteCommand& rhs) const
 	{
