#UE Give option to ask for the ServerDefaultMap if you are not running a server or the GameDefaultMap if you are running a server. Default behavior is unchanged

diff --git a/Engine/Source/Runtime/EngineSettings/Classes/GameMapsSettings.h b/Engine/Source/Runtime/EngineSettings/Classes/GameMapsSettings.h
--- a/Engine/Source/Runtime/EngineSettings/Classes/GameMapsSettings.h
+++ b/Engine/Source/Runtime/EngineSettings/Classes/GameMapsSettings.h
@@ -90,6 +90,12 @@ struct FTemplateMapInfoOverride
 	}
 };
 
+enum class EDefaultMapRequestType : uint8
+{
+	Default,
+	Client,
+	Server
+};
 
 UCLASS(config=Engine, defaultconfig, MinimalAPI)
 class UGameMapsSettings
@@ -99,11 +105,12 @@ class UGameMapsSettings
 
 	/**
 	 * Get the default map specified in the settings.
-	 * Makes a choice based on running as listen server/client vs dedicated server
+	 * Makes a choice based on running as listen server/client vs dedicated server if RequestType is Default
 	 *
+	 * @param RequestType If default, it will return the map for your currently running configuration, otherwise you can request the client or server map specifically
 	 * @return the default map specified in the settings
 	 */
-	static ENGINESETTINGS_API FString GetGameDefaultMap( );
+	static ENGINESETTINGS_API FString GetGameDefaultMap(EDefaultMapRequestType RequestType = EDefaultMapRequestType::Default);
 
 	/**
 	 * Get the global default game type specified in the configuration
diff --git a/Engine/Source/Runtime/EngineSettings/Private/EngineSettingsModule.cpp b/Engine/Source/Runtime/EngineSettings/Private/EngineSettingsModule.cpp
--- a/Engine/Source/Runtime/EngineSettings/Private/EngineSettingsModule.cpp
+++ b/Engine/Source/Runtime/EngineSettings/Private/EngineSettingsModule.cpp
@@ -86,11 +86,22 @@ UHudSettings::UHudSettings( const FObjectInitializer& ObjectInitializer )
 /* Static functions
  *****************************************************************************/
 
-FString UGameMapsSettings::GetGameDefaultMap( )
+FString UGameMapsSettings::GetGameDefaultMap(EDefaultMapRequestType RequestType)
 {
-	return IsRunningDedicatedServer()
-		? GetDefault<UGameMapsSettings>()->ServerDefaultMap.GetLongPackageName()
-		: GetDefault<UGameMapsSettings>()->GameDefaultMap.GetLongPackageName();
+	switch (RequestType)
+	{
+		case EDefaultMapRequestType::Client:
+			return GetDefault<UGameMapsSettings>()->GameDefaultMap.GetLongPackageName();
+
+		case EDefaultMapRequestType::Server:
+			return GetDefault<UGameMapsSettings>()->ServerDefaultMap.GetLongPackageName();
+
+		case EDefaultMapRequestType::Default:
+		default:
+			return IsRunningDedicatedServer()
+				? GetDefault<UGameMapsSettings>()->ServerDefaultMap.GetLongPackageName()
+				: GetDefault<UGameMapsSettings>()->GameDefaultMap.GetLongPackageName();
+	};
 }
 
 
