Add support for eSSE_MONOSCOPIC to FFakeStereoRenderingDevice::GetStereoProjectionMatrix(), and clamp -emulatestereo FOV values to avoid rendering bugs with extreme inputs.

diff --git a/Engine/Source/Runtime/Engine/Private/UnrealEngine.cpp b/Engine/Source/Runtime/Engine/Private/UnrealEngine.cpp
--- a/Engine/Source/Runtime/Engine/Private/UnrealEngine.cpp
+++ b/Engine/Source/Runtime/Engine/Private/UnrealEngine.cpp
@@ -3989,7 +3989,8 @@ public:
 	{
 		float AngleUp = TopFOVRatio * FOVInDegrees;
 		float AngleDown = -(1.0f - TopFOVRatio) * FOVInDegrees;
-		// setup for monoscopic by default
+
+		// Use sane default values
 		float AngleRight = 0.5f * FOVInDegrees;
 		float AngleLeft = -0.5f * FOVInDegrees;
 
@@ -4007,6 +4008,15 @@ public:
 			AngleRight = PrimaryEyeCoeff * FOVInDegrees;
 		}
 
+		// Needed for frustum culling in instanced stereo or mobile-multi view
+		else if (ViewIndex == eSSE_MONOSCOPIC)
+		{
+			// Use the leftmost AngleLeft and the rightmost AngleRight
+			float MaxEyeCoeff = FMath::Max(PrimaryEyeCoeff, SecondaryEyeCoeff);
+			AngleLeft = -(MaxEyeCoeff) * FOVInDegrees;
+			AngleRight = MaxEyeCoeff * FOVInDegrees;
+		}
+
 		AngleUp = tan(FMath::DegreesToRadians(AngleUp));
 		AngleDown = tan(FMath::DegreesToRadians(AngleDown));
 		AngleLeft = tan(FMath::DegreesToRadians(AngleLeft));
diff --git a/Engine/Source/Runtime/EngineSettings/Classes/GeneralProjectSettings.h b/Engine/Source/Runtime/EngineSettings/Classes/GeneralProjectSettings.h
--- a/Engine/Source/Runtime/EngineSettings/Classes/GeneralProjectSettings.h
+++ b/Engine/Source/Runtime/EngineSettings/Classes/GeneralProjectSettings.h
@@ -103,11 +103,15 @@ class UGeneralProjectSettings
 	UPROPERTY(config, EditAnywhere, Category = Settings)
 	float FOVForFakeStereoRenderingDevice;
 
-	/* Determines the how much of -emulatestereo FOV is above the horizontal plane (0-1) - usually helmets have a larger FOV down than up, so this value is less than 0.5 */
-	UPROPERTY(config, EditAnywhere, Category = Settings)
+	/* Determines how much of -emulatestereo FOV is above the horizontal plane (0-1) - usually helmets have a larger FOV down than up, so this value is less than 0.5. */
+	UPROPERTY(config, EditAnywhere, Category = Settings, meta = (
+		ClampMin = "0.01",
+		ClampMax = "0.99"))
 	float TopFOVRatioForFakeStereoRenderingDevice;
 
-	/* Determines difference between left and eyes -emulatestereo FOVs. If 0, they are the same (their FOVs are symmetric), the closer to 1, the more each eye is looking towards its side. Usable values are in 0 - 0.4 range. */
-	UPROPERTY(config, EditAnywhere, Category = Settings)
+	/* Determines difference between left and eyes horizontal -emulatestereo FOVs. If 0, they are the same (their FOVs are symmetric), the closer to 1, the more each eye is looking towards its side. Reasonable values are in 0 - 0.4 range. */
+	UPROPERTY(config, EditAnywhere, Category = Settings, meta=(
+		ClampMin="-0.99",
+		ClampMax="0.99"))
 	float DifferenceBetweenEyesForFakeStereoRenderingDevice;
 };
